Add rich tooltips to brushes in Modes window

diff --git a/Engine/Source/Editor/BspMode/Private/SBspPalette.cpp b/Engine/Source/Editor/BspMode/Private/SBspPalette.cpp
--- a/Engine/Source/Editor/BspMode/Private/SBspPalette.cpp
+++ b/Engine/Source/Editor/BspMode/Private/SBspPalette.cpp
@@ -5,6 +5,7 @@
 #include "BspModeModule.h"
 #include "SScrollBorder.h"
 #include "Editor/UnrealEd/Public/DragAndDrop/BrushBuilderDragDropOp.h"
+#include "EditorClassUtils.h"
 
 #define LOCTEXT_NAMESPACE "BspPalette"
 
@@ -135,8 +136,6 @@ TSharedRef<ITableRow> SBspPalette::MakeListViewWidget(TSharedPtr<FBspBuilderType
 	check(BspBuilder.IsValid());
 	check(BspBuilder->BuilderClass.IsValid());
 
-	UClass* BuilderClass = BspBuilder->BuilderClass.Get();
-	
 	TSharedRef< STableRow<TSharedPtr<FBspBuilderType>> > TableRowWidget = 
 		SNew( STableRow<TSharedPtr<FBspBuilderType>>, OwnerTable )
 		.Style( FBspModeStyle::Get(), "BspMode.TableRow" )
@@ -146,7 +145,7 @@ TSharedRef<ITableRow> SBspPalette::MakeListViewWidget(TSharedPtr<FBspBuilderType
 		SNew(SBorder)
 		.BorderImage(FCoreStyle::Get().GetBrush("NoBorder"))
 		.Padding(0)
-		.ToolTipText(BspBuilder->ToolTipText)
+		.ToolTip(FEditorClassUtils::GetTooltip(ABrush::StaticClass(), BspBuilder->ToolTipText))
 		.Cursor( EMouseCursor::GrabHand )
 		[
 			SNew(SHorizontalBox)
diff --git a/Engine/Source/Editor/UnrealEd/Private/EditorClassUtils.cpp b/Engine/Source/Editor/UnrealEd/Private/EditorClassUtils.cpp
--- a/Engine/Source/Editor/UnrealEd/Private/EditorClassUtils.cpp
+++ b/Engine/Source/Editor/UnrealEd/Private/EditorClassUtils.cpp
@@ -18,7 +18,12 @@ FString FEditorClassUtils::GetDocumentationExcerpt(const UClass* Class)
 
 TSharedRef<SToolTip> FEditorClassUtils::GetTooltip(const UClass* Class)
 {
-	return (Class ? IDocumentation::Get()->CreateToolTip(Class->GetToolTipText(), nullptr, GetDocumentationPage(Class), GetDocumentationExcerpt(Class)) : SNew(SToolTip));
+	return (Class ? GetTooltip(Class, Class->GetToolTipText()) : SNew(SToolTip));
+}
+
+TSharedRef<SToolTip> FEditorClassUtils::GetTooltip(const UClass* Class, const TAttribute<FText>& OverrideText)
+{
+	return (Class ? IDocumentation::Get()->CreateToolTip(OverrideText, nullptr, GetDocumentationPage(Class), GetDocumentationExcerpt(Class)) : SNew(SToolTip));
 }
 
 FString FEditorClassUtils::GetDocumentationLink(const UClass* Class)
diff --git a/Engine/Source/Editor/UnrealEd/Public/EditorClassUtils.h b/Engine/Source/Editor/UnrealEd/Public/EditorClassUtils.h
--- a/Engine/Source/Editor/UnrealEd/Public/EditorClassUtils.h
+++ b/Engine/Source/Editor/UnrealEd/Public/EditorClassUtils.h
@@ -33,6 +33,15 @@ namespace FEditorClassUtils
 	 */
 	UNREALED_API TSharedRef<SToolTip> GetTooltip(const UClass* Class);
 
+	/**
+	 * Gets the tooltip to display for a given class with specified text for the tooltip
+	 *
+	 * @param	InClass			Class we want to build a tooltip for
+	 * @param	OverrideText	The text to display on the standard tooltip
+	 * @return					Shared reference to the constructed tooltip
+	 */
+	UNREALED_API TSharedRef<SToolTip> GetTooltip(const UClass* Class, const TAttribute<FText>& OverrideText);
+
 	/**
 	 * Returns the link path to the documentation for a given class
 	 *
